[Mutable] Fix Populations deprecated warnings.

#rnx

diff --git a/Engine/Plugins/Experimental/MutablePopulation/Source/CustomizableObjectPopulation/Private/MuCOP/CustomizableObjectPopulation.cpp b/Engine/Plugins/Experimental/MutablePopulation/Source/CustomizableObjectPopulation/Private/MuCOP/CustomizableObjectPopulation.cpp
--- a/Engine/Plugins/Experimental/MutablePopulation/Source/CustomizableObjectPopulation/Private/MuCOP/CustomizableObjectPopulation.cpp
+++ b/Engine/Plugins/Experimental/MutablePopulation/Source/CustomizableObjectPopulation/Private/MuCOP/CustomizableObjectPopulation.cpp
@@ -539,8 +539,8 @@ void CompileIntSamplers(
 		// If the parameter is block listed
 		const bool bIsBlocklisted = IsAnyTagInList(ParamTags, PopulationClass->Blocklist);
 
-		const int32 ParamIndex = PopulationClass->CustomizableObject->FindParameter(IntParamName);
-		const int32 NumOptions = PopulationClass->CustomizableObject->GetIntParameterNumOptions(ParamIndex);
+		const int32 ParamIndex = PopulationClass->CustomizableObject->GetPrivate()->FindParameter(IntParamName);
+		const int32 NumOptions = PopulationClass->CustomizableObject->GetPrivate()->GetEnumParameterNumValues(ParamIndex);
 
 		TArray<FString> Options;
 		TArray<TArray<FString>> OptionTags;
@@ -551,7 +551,7 @@ void CompileIntSamplers(
 		// Gather information about options, name and tags.
 		for (int32 I = 0; I < NumOptions; ++I)
 		{
-			FString OptionName = PopulationClass->CustomizableObject->GetIntParameterAvailableOption(ParamIndex, I);
+			FString OptionName = PopulationClass->CustomizableObject->GetPrivate()->GetIntParameterAvailableOption(ParamIndex, I);
 			Options[I] = OptionName;
 
 			if (ParamTagsFound)
@@ -674,7 +674,7 @@ FCustomizableObjectPopulationClassGenerator CompilePopulationClass(UCustomizable
 
 	for (int32 I = 0; I < NumParameters; ++I)
 	{
-		const EMutableParameterType ParamType = ClassObject->GetParameterType(I);
+		const EMutableParameterType ParamType = ClassObject->GetPrivate()->GetParameterType(I);
 
 		switch (ParamType)
 		{
diff --git a/Engine/Plugins/Experimental/MutablePopulation/Source/CustomizableObjectPopulationEditor/Private/MuCOPE/CustomizableObjectPopulationClassDetails.cpp b/Engine/Plugins/Experimental/MutablePopulation/Source/CustomizableObjectPopulationEditor/Private/MuCOPE/CustomizableObjectPopulationClassDetails.cpp
--- a/Engine/Plugins/Experimental/MutablePopulation/Source/CustomizableObjectPopulationEditor/Private/MuCOPE/CustomizableObjectPopulationClassDetails.cpp
+++ b/Engine/Plugins/Experimental/MutablePopulation/Source/CustomizableObjectPopulationEditor/Private/MuCOPE/CustomizableObjectPopulationClassDetails.cpp
@@ -935,7 +935,7 @@ void SPopulationClassConstraint::FillConstraintTypesOptions()
 
 TArray<EPopulationConstraintType> SPopulationClassConstraint::GetConstrinatTypes()
 {
-	switch (PopulationClass->CustomizableObject->GetParameterType(ParameterIndex))
+	switch (PopulationClass->CustomizableObject->GetPrivate()->GetParameterType(ParameterIndex))
 	{
 	case EMutableParameterType::None:
 		break;
@@ -1134,7 +1134,7 @@ void SPopulationClassConstraint::BuildDiscreteWidget()
 	//];
 	
 	// Filling combobox options
-	if (ParameterIndex > -1 && PopulationClass->CustomizableObject->GetIntParameterNumOptions(ParameterIndex) > 0)
+	if (ParameterIndex > -1 && PopulationClass->CustomizableObject->GetPrivate()->GetEnumParameterNumValues(ParameterIndex) > 0)
 	{
 		ParameterOptions.Empty();
 
@@ -1162,9 +1162,9 @@ void SPopulationClassConstraint::BuildDiscreteWidget()
 			]
 		];
 
-		for (int i = 0; i < PopulationClass->CustomizableObject->GetIntParameterNumOptions(ParameterIndex); ++i)
+		for (int i = 0; i < PopulationClass->CustomizableObject->GetPrivate()->GetEnumParameterNumValues(ParameterIndex); ++i)
 		{
-			FString ParameterOption = PopulationClass->CustomizableObject->GetIntParameterAvailableOption(ParameterIndex, i);
+			FString ParameterOption = PopulationClass->CustomizableObject->GetPrivate()->GetIntParameterAvailableOption(ParameterIndex, i);
 
 			ParameterOptions.Add(MakeShareable(new FString(ParameterOption)));
 			if (ParameterOption == PopulationClass->Characteristics[CharactericticIndex].Constraints[ConstraintIndex].DiscreteValue)
diff --git a/Engine/Plugins/Experimental/MutablePopulation/Source/CustomizableObjectPopulationEditor/Private/MuCOPE/SCustomizableObjectPopulationClassTagsTool.cpp b/Engine/Plugins/Experimental/MutablePopulation/Source/CustomizableObjectPopulationEditor/Private/MuCOPE/SCustomizableObjectPopulationClassTagsTool.cpp
--- a/Engine/Plugins/Experimental/MutablePopulation/Source/CustomizableObjectPopulationEditor/Private/MuCOPE/SCustomizableObjectPopulationClassTagsTool.cpp
+++ b/Engine/Plugins/Experimental/MutablePopulation/Source/CustomizableObjectPopulationEditor/Private/MuCOPE/SCustomizableObjectPopulationClassTagsTool.cpp
@@ -440,12 +440,12 @@ void SCustomizableObjectPopulationClassTagsTool::GenerateParameterOptionsManager
 		{
 			if (SelectedParameter == PopulationClass->CustomizableObject->GetParameterName(i))
 			{
-				if (PopulationClass->CustomizableObject->GetIntParameterNumOptions(i) > 0)
+				if (PopulationClass->CustomizableObject->GetPrivate()->GetEnumParameterNumValues(i) > 0)
 				{
-					for (int32 j = 0; j < PopulationClass->CustomizableObject->GetIntParameterNumOptions(i); ++j)
+					for (int32 j = 0; j < PopulationClass->CustomizableObject->GetPrivate()->GetEnumParameterNumValues(i); ++j)
 					{
 						TSharedPtr<SHorizontalBox> VerticalWidget = SNew(SHorizontalBox);
-						FString ParameterOptionName = PopulationClass->CustomizableObject->GetIntParameterAvailableOption(i, j);
+						FString ParameterOptionName = PopulationClass->CustomizableObject->GetPrivate()->GetIntParameterAvailableOption(i, j);
 
 						VerticalWidget->AddSlot()
 						.AutoWidth()
