[Mutable] Remove deprecated usages of Find<Type>Parameter functions in Mutable Population.

diff --git a/Engine/Plugins/Experimental/MutablePopulation/Source/CustomizableObjectPopulation/Private/MuCOP/CustomizableObjectPopulationGenerator.cpp b/Engine/Plugins/Experimental/MutablePopulation/Source/CustomizableObjectPopulation/Private/MuCOP/CustomizableObjectPopulationGenerator.cpp
--- a/Engine/Plugins/Experimental/MutablePopulation/Source/CustomizableObjectPopulation/Private/MuCOP/CustomizableObjectPopulationGenerator.cpp
+++ b/Engine/Plugins/Experimental/MutablePopulation/Source/CustomizableObjectPopulation/Private/MuCOP/CustomizableObjectPopulationGenerator.cpp
@@ -36,7 +36,7 @@ void FCustomizableObjectPopulationClassGenerator::GenerateParameters(
 			{
 				const bool SampledValue = BoolSamplers[SamplerIndex.SamplerIndex].Sample(Rand);
 
-				if (Instance->FindBoolParameterNameIndex(ParamName) != INDEX_NONE)
+				if (Instance->ContainsBoolParameter(ParamName))
 				{
 					Instance->SetBoolParameterSelectedOption(ParamName, SampledValue);
 				}
@@ -49,7 +49,7 @@ void FCustomizableObjectPopulationClassGenerator::GenerateParameters(
 				if (IntParamSampler.IsValidSampler())
 				{
 					const FString& OptionName = IntParamSampler.GetOptionName(IntParamSampler.Sample(Rand));
-					if (Instance->FindIntParameterNameIndex(ParamName) != INDEX_NONE)
+					if (Instance->ContainsIntParameter(ParamName))
 					{
 						if (!CustomizableObject->IsParameterMultidimensional(ParamName))
 						{
@@ -66,7 +66,7 @@ void FCustomizableObjectPopulationClassGenerator::GenerateParameters(
 			case EPopulationSamplerType::UNIFORM_FLOAT:
 			{
 				const float SampledValue = UniformFloatSamplers[SamplerIndex.SamplerIndex].Sample(Rand);
-				if (Instance->FindFloatParameterNameIndex(ParamName) != INDEX_NONE)
+				if (Instance->ContainsFloatParameter(ParamName))
 				{
 					if (!CustomizableObject->IsParameterMultidimensional(ParamName))
 					{
@@ -82,7 +82,7 @@ void FCustomizableObjectPopulationClassGenerator::GenerateParameters(
 			case EPopulationSamplerType::CURVE:
 			{
 				const float SampledValue = CurveSamplers[SamplerIndex.SamplerIndex].Sample(Rand);
-				if (Instance->FindFloatParameterNameIndex(ParamName) != INDEX_NONE)
+				if (Instance->ContainsFloatParameter(ParamName))
 				{
 					Instance->SetFloatParameterSelectedOption(ParamName, SampledValue);
 				}
@@ -91,7 +91,7 @@ void FCustomizableObjectPopulationClassGenerator::GenerateParameters(
 			case EPopulationSamplerType::RANGE:
 			{
 				const float SampledValue = RangesSamplers[SamplerIndex.SamplerIndex].Sample(Rand);
-				if (Instance->FindFloatParameterNameIndex(ParamName) != INDEX_NONE)
+				if (Instance->ContainsFloatParameter(ParamName))
 				{
 					Instance->SetFloatParameterSelectedOption(ParamName, SampledValue);
 				}
@@ -100,7 +100,7 @@ void FCustomizableObjectPopulationClassGenerator::GenerateParameters(
 			case EPopulationSamplerType::UNIFORM_CURVE_COLOR:
 			{
 				const FLinearColor SampledValue = UniformColorSamplers[SamplerIndex.SamplerIndex].Sample(Rand);
-				if (Instance->FindVectorParameterNameIndex(ParamName) != INDEX_NONE)
+				if (Instance->ContainsVectorParameter(ParamName))
 				{
 					Instance->SetColorParameterSelectedOption(ParamName, SampledValue);
 				}
@@ -108,7 +108,7 @@ void FCustomizableObjectPopulationClassGenerator::GenerateParameters(
 			}
 			case EPopulationSamplerType::CONSTANT_COLOR:
 			{
-				if (Instance->FindVectorParameterNameIndex(ParamName) != INDEX_NONE)
+				if (Instance->ContainsVectorParameter(ParamName))
 				{
 					Instance->SetColorParameterSelectedOption(ParamName, ConstantColors[SamplerIndex.SamplerIndex]);
 				}
