[mutable] Changed the title of the tab of all CustomizableObject editors so just the asset being edited is shown, to be more like the other UE Editors.
#tests Editor
#rnx

diff --git a/Engine/Plugins/Experimental/MutablePopulation/Source/CustomizableObjectPopulationEditor/Private/MuCOPE/CustomizableObjectPopulationClassEditor.cpp b/Engine/Plugins/Experimental/MutablePopulation/Source/CustomizableObjectPopulationEditor/Private/MuCOPE/CustomizableObjectPopulationClassEditor.cpp
--- a/Engine/Plugins/Experimental/MutablePopulation/Source/CustomizableObjectPopulationEditor/Private/MuCOPE/CustomizableObjectPopulationClassEditor.cpp
+++ b/Engine/Plugins/Experimental/MutablePopulation/Source/CustomizableObjectPopulationEditor/Private/MuCOPE/CustomizableObjectPopulationClassEditor.cpp
@@ -233,10 +233,7 @@ FText FCustomizableObjectPopulationClassEditor::GetBaseToolkitName() const
 
 FText FCustomizableObjectPopulationClassEditor::GetToolkitName() const
 {
-	FFormatNamedArguments Args;
-	Args.Add(TEXT("ObjectName"), FText::FromString(GetEditingObject()->GetName()));
-	Args.Add(TEXT("ToolkitName"), GetBaseToolkitName());
-	return FText::Format(LOCTEXT("AppLabelWithAssetName", "{ObjectName} - {ToolkitName}"), Args);
+	return FText::FromString(GetEditingObject()->GetName());
 }
 
 
diff --git a/Engine/Plugins/Experimental/MutablePopulation/Source/CustomizableObjectPopulationEditor/Private/MuCOPE/CustomizableObjectPopulationEditor.cpp b/Engine/Plugins/Experimental/MutablePopulation/Source/CustomizableObjectPopulationEditor/Private/MuCOPE/CustomizableObjectPopulationEditor.cpp
--- a/Engine/Plugins/Experimental/MutablePopulation/Source/CustomizableObjectPopulationEditor/Private/MuCOPE/CustomizableObjectPopulationEditor.cpp
+++ b/Engine/Plugins/Experimental/MutablePopulation/Source/CustomizableObjectPopulationEditor/Private/MuCOPE/CustomizableObjectPopulationEditor.cpp
@@ -149,10 +149,7 @@ FText FCustomizableObjectPopulationEditor::GetBaseToolkitName() const
 
 FText FCustomizableObjectPopulationEditor::GetToolkitName() const
 {
-	FFormatNamedArguments Args;
-	Args.Add(TEXT("ObjectName"), FText::FromString(GetEditingObject()->GetName()));
-	Args.Add(TEXT("ToolkitName"), GetBaseToolkitName());
-	return FText::Format(LOCTEXT("AppLabelWithAssetName", "{ObjectName} - {ToolkitName}"), Args);
+	return FText::FromString(GetEditingObject()->GetName());
 }
 
 
diff --git a/Engine/Plugins/Mutable/Source/CustomizableObjectEditor/Private/MuCOE/CustomizableObjectEditor.cpp b/Engine/Plugins/Mutable/Source/CustomizableObjectEditor/Private/MuCOE/CustomizableObjectEditor.cpp
--- a/Engine/Plugins/Mutable/Source/CustomizableObjectEditor/Private/MuCOE/CustomizableObjectEditor.cpp
+++ b/Engine/Plugins/Mutable/Source/CustomizableObjectEditor/Private/MuCOE/CustomizableObjectEditor.cpp
@@ -1187,10 +1187,7 @@ TSharedRef<SWidget> FCustomizableObjectEditor::GenerateCompileOptionsMenuContent
 
 FText FCustomizableObjectEditor::GetToolkitName() const
 {
-	FFormatNamedArguments Args;
-	Args.Add(TEXT("ObjectName"), FText::FromString(GetEditingObject()->GetName()));
-	Args.Add(TEXT("ToolkitName"), GetBaseToolkitName());
-	return FText::Format(LOCTEXT("AppLabelWithAssetName", "{ObjectName} - {ToolkitName}"), Args);
+	return FText::FromString(GetEditingObject()->GetName());
 }
 
 
diff --git a/Engine/Plugins/Mutable/Source/CustomizableObjectEditor/Private/MuCOE/CustomizableObjectInstanceEditor.cpp b/Engine/Plugins/Mutable/Source/CustomizableObjectEditor/Private/MuCOE/CustomizableObjectInstanceEditor.cpp
--- a/Engine/Plugins/Mutable/Source/CustomizableObjectEditor/Private/MuCOE/CustomizableObjectInstanceEditor.cpp
+++ b/Engine/Plugins/Mutable/Source/CustomizableObjectEditor/Private/MuCOE/CustomizableObjectInstanceEditor.cpp
@@ -520,10 +520,7 @@ void FCustomizableObjectInstanceEditor::ExtendToolbar()
 
 FText FCustomizableObjectInstanceEditor::GetToolkitName() const
 {
-	FFormatNamedArguments Args;
-	Args.Add(TEXT("ObjectName"), FText::FromString(GetEditingObject()->GetName()));
-	Args.Add(TEXT("ToolkitName"), GetBaseToolkitName());
-	return FText::Format(LOCTEXT("AppLabelWithAssetName", "{ObjectName} - {ToolkitName}"), Args);
+	return FText::FromString(GetEditingObject()->GetName());
 }
 
 
