[mutable population] Fixed random crash when using 3 or more mutable population classes in one mutable population object.
- Added support for single non-numerical COI component.

#rnx

diff --git a/Engine/Plugins/Experimental/MutablePopulation/Source/CustomizableObjectPopulation/Private/MuCOP/CustomizableObjectPopulationGenerator.cpp b/Engine/Plugins/Experimental/MutablePopulation/Source/CustomizableObjectPopulation/Private/MuCOP/CustomizableObjectPopulationGenerator.cpp
--- a/Engine/Plugins/Experimental/MutablePopulation/Source/CustomizableObjectPopulation/Private/MuCOP/CustomizableObjectPopulationGenerator.cpp
+++ b/Engine/Plugins/Experimental/MutablePopulation/Source/CustomizableObjectPopulation/Private/MuCOP/CustomizableObjectPopulationGenerator.cpp
@@ -201,7 +201,7 @@ void UCustomizableObjectPopulationGenerator::Init(TArray<UCustomizableObject*> _
 	}
 }
 
-bool UCustomizableObjectPopulationGenerator::IsCompiledForCook()
+bool UCustomizableObjectPopulationGenerator::IsCompiledForCook() const
 {
 	return Private.IsValid() && Private->bPopulationCompiledForCook;
 }
@@ -239,10 +239,10 @@ void UCustomizableObjectPopulationGenerator::Generate(
 		}
 	}
 
-	// Needed to keep the index of the eddited instance in the inner loops
+	// Needed to keep the index of the edited instance in the inner loops
 	int InstanceIndex = 0;
 
-	// Set Intances Objects.
+	// Set Instance Objects.
 	for (int32 I = 0; I < NumClasses; ++I)
 	{
 		UCustomizableObject* ClassObject = PopulationObjects[I];
@@ -253,14 +253,14 @@ void UCustomizableObjectPopulationGenerator::Generate(
 			OutPopulationInstances[C + InstanceIndex]->SetObject(ClassObject);
 		}
 
-		InstanceIndex = NumClassInstances;
+		InstanceIndex += NumClassInstances;
 	}
 
 	FRandomStream ParamRand(Seed);
 
 	InstanceIndex = 0;
 
-	// Generate Intances Parameters.
+	// Generate Instance Parameters.
 	for (int32 I = 0; I < NumClasses; ++I)
 	{
 		FCustomizableObjectPopulationClassGenerator& ClassGenerator = Private->ClassGenerators[I];
@@ -271,6 +271,6 @@ void UCustomizableObjectPopulationGenerator::Generate(
 			ClassGenerator.GenerateParameters(ParamRand, OutPopulationInstances[C + InstanceIndex]);
 		}
 
-		InstanceIndex = NumClassInstances;
+		InstanceIndex += NumClassInstances;
 	}
 }
diff --git a/Engine/Plugins/Experimental/MutablePopulation/Source/CustomizableObjectPopulation/Public/MuCOP/CustomizableObjectPopulationGenerator.h b/Engine/Plugins/Experimental/MutablePopulation/Source/CustomizableObjectPopulation/Public/MuCOP/CustomizableObjectPopulationGenerator.h
--- a/Engine/Plugins/Experimental/MutablePopulation/Source/CustomizableObjectPopulation/Public/MuCOP/CustomizableObjectPopulationGenerator.h
+++ b/Engine/Plugins/Experimental/MutablePopulation/Source/CustomizableObjectPopulation/Public/MuCOP/CustomizableObjectPopulationGenerator.h
@@ -76,7 +76,7 @@ public:
 #if WITH_EDITOR
 	void Init(TArray<UCustomizableObject*> _PopulationObjects, TSharedPtr<FCustomizableObjectPopulationGeneratorPrivate> _Private);
 
-	bool IsCompiledForCook();
+	bool IsCompiledForCook() const;
 #endif
 
 private:
diff --git a/Engine/Plugins/Experimental/MutablePopulation/Source/CustomizableObjectPopulationEditor/Private/MuCOPE/CustomizableObjectPopulationClassEditor.cpp b/Engine/Plugins/Experimental/MutablePopulation/Source/CustomizableObjectPopulationEditor/Private/MuCOPE/CustomizableObjectPopulationClassEditor.cpp
--- a/Engine/Plugins/Experimental/MutablePopulation/Source/CustomizableObjectPopulationEditor/Private/MuCOPE/CustomizableObjectPopulationClassEditor.cpp
+++ b/Engine/Plugins/Experimental/MutablePopulation/Source/CustomizableObjectPopulationEditor/Private/MuCOPE/CustomizableObjectPopulationClassEditor.cpp
@@ -772,6 +772,12 @@ void FCustomizableObjectPopulationClassEditor::TestPopulationClass()
 				{
 					for (int32 i = 0; i < ViewportInstances.Num(); ++i)
 					{
+						const UCustomizableObject* CustomizableObject = ViewportInstances[i]->GetCustomizableObject();	
+						if (!CustomizableObject)
+						{
+							continue;
+						}
+
 						// Creating the Customizable Skeletal Component
 						UCustomizableSkeletalComponent* PreviewCustomizableSkeletalComponent = nullptr;
 						PreviewCustomizableSkeletalComponent = NewObject<UCustomizableSkeletalComponent>(UCustomizableSkeletalComponent::StaticClass());
@@ -791,6 +797,8 @@ void FCustomizableObjectPopulationClassEditor::TestPopulationClass()
 							ViewportInstances[i]->UpdateSkeletalMeshAsync(true, true);
 
 							PreviewCustomizableSkeletalComponent->CustomizableObjectInstance = ViewportInstances[i];
+							// TODO: UE-274335. We want this code to be able to operate with multiple components, not just the first one
+							PreviewCustomizableSkeletalComponent->SetComponentName(CustomizableObject->GetComponentName(0));
 							PreviewCustomizableSkeletalComponent->AttachToComponent(SkeletalMeshComponent, FAttachmentTransformRules::KeepRelativeTransform);
 
 							ColliderComponent->AttachToComponent(SkeletalMeshComponent, FAttachmentTransformRules::KeepRelativeTransform);
diff --git a/Engine/Plugins/Experimental/MutablePopulation/Source/CustomizableObjectPopulationEditor/Private/MuCOPE/CustomizableObjectPopulationEditor.cpp b/Engine/Plugins/Experimental/MutablePopulation/Source/CustomizableObjectPopulationEditor/Private/MuCOPE/CustomizableObjectPopulationEditor.cpp
--- a/Engine/Plugins/Experimental/MutablePopulation/Source/CustomizableObjectPopulationEditor/Private/MuCOPE/CustomizableObjectPopulationEditor.cpp
+++ b/Engine/Plugins/Experimental/MutablePopulation/Source/CustomizableObjectPopulationEditor/Private/MuCOPE/CustomizableObjectPopulationEditor.cpp
@@ -453,6 +453,12 @@ void FCustomizableObjectPopulationEditor::TestPopulation()
 		{
 			for (int32 i = 0; i < ViewportInstances.Num(); ++i)
 			{
+				const UCustomizableObject* CustomizableObject = ViewportInstances[i]->GetCustomizableObject();	
+				if (!CustomizableObject)
+				{
+					continue;
+				}
+
 				// Creating the Customizable Skeletal Component
 				UCustomizableSkeletalComponent* PreviewCustomizableSkeletalComponent = nullptr;
 				PreviewCustomizableSkeletalComponent = NewObject<UCustomizableSkeletalComponent>(UCustomizableSkeletalComponent::StaticClass());
@@ -472,6 +478,8 @@ void FCustomizableObjectPopulationEditor::TestPopulation()
 					ViewportInstances[i]->UpdateSkeletalMeshAsync(true, true);
 
 					PreviewCustomizableSkeletalComponent->CustomizableObjectInstance = ViewportInstances[i];
+					// TODO: UE-274335. We want this code to be able to operate with multiple components, not just the first one
+					PreviewCustomizableSkeletalComponent->SetComponentName(CustomizableObject->GetComponentName(0));
 					PreviewCustomizableSkeletalComponent->AttachToComponent(PreviewSkeletalMeshComponent, FAttachmentTransformRules::KeepRelativeTransform);
 
 					ColliderComponent->AttachToComponent(PreviewSkeletalMeshComponent, FAttachmentTransformRules::KeepRelativeTransform);
